Static Mesh Asset Editor: Initial commit of new asset editor for static mesh assets based on the modern asset editor framework

#Synced-CL 49048852

diff --git a/Engine/Source/Developer/GeometryProcessingInterfaces/Public/GeometryProcessingInterfaces/IStaticMeshAssetEditorModularFeature.h b/Engine/Source/Developer/GeometryProcessingInterfaces/Public/GeometryProcessingInterfaces/IStaticMeshAssetEditorModularFeature.h
new file mode 100644
--- /dev/null
+++ b/Engine/Source/Developer/GeometryProcessingInterfaces/Public/GeometryProcessingInterfaces/IStaticMeshAssetEditorModularFeature.h
@@ -0,0 +1,33 @@
+// Copyright Epic Games, Inc. All Rights Reserved.
+
+#pragma once
+
+#include "CoreMinimal.h"
+#include "Features/IModularFeature.h"
+
+/**
+ * This interface acts as a connector to the (optional) Static Mesh Asset Editor plugin. If the plugin is available,
+ * then it can be retrieved via IModularFeatures::Get().GetModularFeatureImplementation or related methods.
+ */
+class IStaticMeshAssetEditorModularFeature : public IModularFeature
+{
+public:
+	virtual ~IStaticMeshAssetEditorModularFeature() {}
+
+	virtual void LaunchStaticMeshAssetEditor(const TArray<TObjectPtr<UObject>>& Objects)
+	{
+		check(false);		// not implemented in base class
+	}
+
+	virtual bool CanLaunchStaticMeshAssetEditor(const TArray<TObjectPtr<UObject>>& Objects)
+	{
+		check(false);		// not implemented in base class
+		return false;
+	}
+
+	// Modular feature name to register for retrieval during runtime
+	static const FName GetModularFeatureName()
+	{
+		return TEXT("StaticMeshAssetEditorModularFeature");
+	}
+};
